DMX - Add a plugin that overrides the datasmith native translator so that MVR is imported along with the datasmith scene, if an MVR file is present in its assets

diff --git a/Engine/Plugins/VirtualProduction/DMX/DatasmithMVR/DatasmithMVR.uplugin b/Engine/Plugins/VirtualProduction/DMX/DatasmithMVR/DatasmithMVR.uplugin
new file mode 100644
--- /dev/null
+++ b/Engine/Plugins/VirtualProduction/DMX/DatasmithMVR/DatasmithMVR.uplugin
@@ -0,0 +1,36 @@
+{
+	"FileVersion": 1,
+	"Version": 1,
+	"VersionName": "1.0",
+	"FriendlyName": "Datasmith MVR",
+	"Description": "Enabled support to import .udatasmith files with MVR content.",
+	"Category": "Virtual Production",
+	"CreatedBy": "Epic Games, Inc.",
+	"CreatedByURL": "http://epicgames.com",
+	"DocsURL": "",
+	"MarketplaceURL": "",
+	"SupportURL": "",
+	"CanContainContent": false,
+	"IsBetaVersion": true,
+	"Modules": [
+		{
+			"Name": "DatasmithMVRTranslator",
+			"Type": "Editor",
+			"LoadingPhase": "Default"
+		}
+	],
+	"Plugins": [
+		{
+			"Name": "DMXEngine",
+			"Enabled": true
+		},
+		{
+			"Name": "DatasmithContent",
+			"Enabled": true
+		},
+		{
+			"Name": "DatasmithImporter",
+			"Enabled": true
+		},
+	]
+}
diff --git a/Engine/Plugins/VirtualProduction/DMX/DatasmithMVR/Source/DatasmithMVRTranslator/DatasmithMVRTranslator.Build.cs b/Engine/Plugins/VirtualProduction/DMX/DatasmithMVR/Source/DatasmithMVRTranslator/DatasmithMVRTranslator.Build.cs
new file mode 100644
--- /dev/null
+++ b/Engine/Plugins/VirtualProduction/DMX/DatasmithMVR/Source/DatasmithMVRTranslator/DatasmithMVRTranslator.Build.cs
@@ -0,0 +1,32 @@
+// Copyright Epic Games, Inc. All Rights Reserved.
+
+using UnrealBuildTool;
+
+public class DatasmithMVRTranslator : ModuleRules
+{
+	public DatasmithMVRTranslator(ReadOnlyTargetRules Target) : base(Target)
+	{
+		PCHUsage = PCHUsageMode.UseExplicitOrSharedPCHs;
+
+		PrivateDependencyModuleNames.AddRange(new string[]
+		{
+			"Core",
+			"CoreUObject",
+		});
+
+
+		PrivateDependencyModuleNames.AddRange(
+			new string[]
+			{
+				"AssetRegistry",
+				"DatasmithContent",
+				"DatasmithCore",
+				"DatasmithNativeTranslator",
+				"DatasmithTranslator",
+				"DMXEditor",
+				"DMXRuntime",
+				"Engine",
+			}
+		);
+	}
+}
diff --git a/Engine/Plugins/VirtualProduction/DMX/DatasmithMVR/Source/DatasmithMVRTranslator/Private/DatasmithMVRImportOptions.h b/Engine/Plugins/VirtualProduction/DMX/DatasmithMVR/Source/DatasmithMVRTranslator/Private/DatasmithMVRImportOptions.h
new file mode 100644
--- /dev/null
+++ b/Engine/Plugins/VirtualProduction/DMX/DatasmithMVR/Source/DatasmithMVRTranslator/Private/DatasmithMVRImportOptions.h
@@ -0,0 +1,22 @@
+// Copyright Epic Games, Inc. All Rights Reserved.
+
+#pragma once
+
+#include "DatasmithImportOptions.h"
+
+#include "CoreMinimal.h"
+
+#include "DatasmithMVRImportOptions.generated.h"
+
+
+UCLASS(Config = EditorPerProjectUserSettings)
+class UDatasmithMVRImportOptions 
+	: public UDatasmithOptionsBase
+{
+	GENERATED_BODY()
+
+public:
+	/** If set to true, datasmith elements that correspond to an entry in the MVR are replaced with GDTF Actors */
+	UPROPERTY(Config, EditAnywhere, Category = MVR)
+	bool bImportMVR = true;
+};
diff --git a/Engine/Plugins/VirtualProduction/DMX/DatasmithMVR/Source/DatasmithMVRTranslator/Private/DatasmithMVRNativeTranslator.cpp b/Engine/Plugins/VirtualProduction/DMX/DatasmithMVR/Source/DatasmithMVRTranslator/Private/DatasmithMVRNativeTranslator.cpp
new file mode 100644
--- /dev/null
+++ b/Engine/Plugins/VirtualProduction/DMX/DatasmithMVR/Source/DatasmithMVRTranslator/Private/DatasmithMVRNativeTranslator.cpp
@@ -0,0 +1,240 @@
+// Copyright Epic Games, Inc. All Rights Reserved.
+
+#include "DatasmithMVRNativeTranslator.h"
+
+#include "DatasmithMVRImportOptions.h"
+#include "Factories/DMXLibraryFromMVRFactory.h"
+#include "Library/DMXEntityFixturePatch.h"
+#include "Library/DMXLibrary.h"
+
+#include "Editor.h"
+#include "AssetRegistry/AssetRegistryModule.h"
+#include "DatasmithNativeTranslator.h"
+#include "DatasmithSceneFactory.h"
+#include "IDatasmithSceneElements.h"
+#include "Subsystems/ImportSubsystem.h"
+
+
+DECLARE_LOG_CATEGORY_CLASS(LogDatasmithMVRNativeTranslator, Log, All);
+
+bool FDatasmithMVRNativeTranslator::LoadScene(TSharedRef<IDatasmithScene> InOutScene)
+{
+	const bool bNativeTranslatorSuccess = FDatasmithNativeTranslator::LoadScene(InOutScene);
+	if (!bNativeTranslatorSuccess)
+	{
+		return false;
+	}
+
+	if (ImportOptions.IsValid() && !ImportOptions->bImportMVR)
+	{
+		// Importing MVR is turned off in options. Acknowledge the scene was successfully imported by the native translator.
+		return true;
+	}
+
+	InOutScene->SetHost(TEXT("DatasmithMVRNativeTranslator"));
+	InOutScene->SetProductName(TEXT("DatasmithMVRNativeTranslator"));
+
+	FString MVRFilePathAndName;
+	if (!FindMVRFile(MVRFilePathAndName))
+	{
+		// If there is no occurence of MVR, acknowledge the scene was successfully imported by the native translator.
+		return true;
+	}
+
+	UDMXLibrary* DMXLibrary = CreateDMXLibraryFromMVR(MVRFilePathAndName);
+	if (!DMXLibrary)
+	{		
+		const FString MVRFilename = FPaths::GetCleanFilename(MVRFilePathAndName);
+		UE_LOG(LogDatasmithMVRNativeTranslator, Warning, TEXT("Failed to create DMX Library for %s. See previous errors for more info. Creating native datasmith scene instead."), *MVRFilename);
+		return true;
+	}
+
+	ReplaceMVRActorsWithMVRSceneActor(InOutScene, DMXLibrary);
+
+	return true;
+}
+
+void FDatasmithMVRNativeTranslator::GetSceneImportOptions(TArray<TStrongObjectPtr<UDatasmithOptionsBase>>& Options)
+{
+	FDatasmithNativeTranslator::GetSceneImportOptions(Options);
+
+	if (!ImportOptions.IsValid())
+	{
+		const FString& FilePath = GetSource().GetSourceFile();
+
+		ImportOptions = Datasmith::MakeOptions<UDatasmithMVRImportOptions>();
+	}
+
+	Options.Add(ImportOptions);
+}
+
+void FDatasmithMVRNativeTranslator::SetSceneImportOptions(TArray<TStrongObjectPtr<UDatasmithOptionsBase>>& Options)
+{
+	FDatasmithNativeTranslator::SetSceneImportOptions(Options);
+
+	for (const TStrongObjectPtr<UDatasmithOptionsBase>& OptionPtr : Options)
+	{
+		if (UDatasmithMVRImportOptions* InImportOptions = Cast<UDatasmithMVRImportOptions>(OptionPtr.Get()))
+		{
+			ImportOptions.Reset(InImportOptions);
+		}
+	}
+}
+
+bool FDatasmithMVRNativeTranslator::FindMVRFile(FString& OutMVRFilePathAndName) const
+{
+	const FString DatasmithFilePathAndName = GetSource().GetSourceFile();
+	const FString DatasmithFileName = FPaths::GetBaseFilename(DatasmithFilePathAndName);
+	const FString DatasmithPath = FPaths::GetPath(DatasmithFilePathAndName);
+	const FString DatasmithAssetsPath = FPaths::GetPath(DatasmithFilePathAndName) / DatasmithFileName + TEXT("_Assets");
+
+	const FString MVRFilename = FPaths::GetBaseFilename(DatasmithFilePathAndName) + TEXT(".mvr");
+	
+	if (FPaths::FileExists(DatasmithPath / MVRFilename))
+	{
+		OutMVRFilePathAndName = DatasmithPath / MVRFilename;
+		return true;
+	}
+	else if (FPaths::FileExists(DatasmithAssetsPath / MVRFilename))
+	{
+		OutMVRFilePathAndName = DatasmithAssetsPath / MVRFilename;
+		return true;;
+	}
+
+	return false;
+}
+
+UDMXLibrary* FDatasmithMVRNativeTranslator::CreateDMXLibraryFromMVR(const FString& MVRFilePathAndName) const
+{
+	if (!ensureMsgf(FPaths::FileExists(MVRFilePathAndName), TEXT("Tried to create a DMX Library from MVR, but the MVR file doesn't exist anymore.")))
+	{
+		return nullptr;
+	}
+
+	UDMXLibraryFromMVRFactory* DMXLibraryFromMVRFactory = NewObject<UDMXLibraryFromMVRFactory>();
+
+	const FString SceneName = GetSource().GetSceneName();
+	const FString DMXLibraryName = FPaths::GetBaseFilename(MVRFilePathAndName);
+	const FString PackageName = FPaths::ProjectContentDir() / SceneName;
+	UPackage* DMXLibraryAssetPackage = CreatePackage(*PackageName);
+	DMXLibraryAssetPackage->FullyLoad();
+	
+	bool bOperationCanceled;
+	UObject* NewDMXLibraryObject = DMXLibraryFromMVRFactory->FactoryCreateFile(UDMXLibrary::StaticClass(), DMXLibraryAssetPackage, *DMXLibraryName, RF_Public | RF_Standalone | RF_Transactional, MVRFilePathAndName, nullptr, GWarn, bOperationCanceled);
+	if (bOperationCanceled)
+	{
+		return nullptr;
+	}
+
+	if (UDMXLibrary* DMXLibrary = Cast<UDMXLibrary>(NewDMXLibraryObject))
+	{
+		FAssetRegistryModule::AssetCreated(DMXLibrary);
+		DMXLibrary->MarkPackageDirty();
+
+		return DMXLibrary;
+	}
+
+	return nullptr;
+}
+
+void FDatasmithMVRNativeTranslator::ReplaceMVRActorsWithMVRSceneActor(TSharedRef<IDatasmithScene>& InOutScene, UDMXLibrary* DMXLibrary) const
+{
+	if (!ensureMsgf(DMXLibrary, TEXT("Tried to create an MVR scene in the datasmith scene, but DMXLibrary is null.")))
+	{
+		return;
+	}
+
+	// Create the custom MVR Scene Actor element
+	static const FString MVRSceneActorName = TEXT("MVR Scene Actor");
+	static const FString MVRSceneActorPath = TEXT("/Script/DMXRuntime.DMXMVRSceneActor");
+
+	const TSharedRef<IDatasmithCustomActorElement> CustomMVRSceneActorElement = FDatasmithSceneFactory::CreateCustomActor(*MVRSceneActorName);
+	InOutScene->AddActor(CustomMVRSceneActorElement);
+	CustomMVRSceneActorElement->SetClassOrPathName(*MVRSceneActorPath);
+	CustomMVRSceneActorElement->SetIsAComponent(false);
+	CustomMVRSceneActorElement->SetRotation(FQuat::Identity);
+	CustomMVRSceneActorElement->SetScale(FVector::OneVector);
+	CustomMVRSceneActorElement->SetTranslation(FVector::ZeroVector);
+
+	// Set the DMX Library path as metadata of the custom actor element
+	const TSoftObjectPtr<UDMXEntityFixturePatch> SoftDMXLibraryObject = DMXLibrary;
+	const TSharedRef<IDatasmithKeyValueProperty> FixturePatchKeyValueProperty = FDatasmithSceneFactory::CreateKeyValueProperty(TEXT("DMXLibraryPath"));
+	FixturePatchKeyValueProperty->SetValue(*SoftDMXLibraryObject.ToString());
+
+	const TSharedRef<IDatasmithMetaDataElement> CustomActorMetaDataElement = FDatasmithSceneFactory::CreateMetaData(CustomMVRSceneActorElement->GetName());
+	CustomActorMetaDataElement->AddProperty(FixturePatchKeyValueProperty);
+	CustomActorMetaDataElement->SetAssociatedElement(CustomMVRSceneActorElement);
+
+	InOutScene->AddMetaData(CustomActorMetaDataElement);
+
+	// Remove native datasmith elements that correspond to an MVR Fixtures
+	TArray<TSharedRef<IDatasmithMetaDataElement>> MetaDataElementsToRemove;
+	TArray<TSharedRef<IDatasmithActorElement>> ActorElementsToRemove;
+
+	const int32 NumMetaDatas = InOutScene->GetMetaDataCount();
+	TArray<UDMXEntityFixturePatch*> FixturePatches = DMXLibrary->GetEntitiesTypeCast<UDMXEntityFixturePatch>();
+	for (UDMXEntityFixturePatch* FixturePatch : FixturePatches)
+	{
+		const TArray<UDMXMVRFixtureInstance*>& MVRFixtureInstances = FixturePatch->GetMVRFixtureInstances();
+		for (UDMXMVRFixtureInstance* MVRFixtureInstance : MVRFixtureInstances)
+		{
+			const FDMXMVRFixture& MVRFixture = MVRFixtureInstance->GetMVRFixture();
+			for (int32 MetaDataIndex = 0; MetaDataIndex < NumMetaDatas; MetaDataIndex++)
+			{
+				const TSharedPtr<IDatasmithMetaDataElement> MetaDataElement = InOutScene->GetMetaData(MetaDataIndex);
+
+				// Try to find an MVR ID in MetaData. Note Vectorworks uses VW_UUID as that, so we search for that too.
+				static const FString Generic_MVRUUIDName = TEXT("MVR_UUID");
+				static const FString VW_MVRUUIDName = TEXT("VW_UUID");
+
+				const TSharedPtr<IDatasmithKeyValueProperty> DatasmithKeyValueProperty = [&MetaDataElement]()
+				{
+					TSharedPtr<IDatasmithKeyValueProperty> Result = MetaDataElement->GetPropertyByName(*Generic_MVRUUIDName);
+					if (!Result.IsValid())
+					{
+						Result = MetaDataElement->GetPropertyByName(*VW_MVRUUIDName);
+					}
+
+					return Result;
+				}();
+
+				if (!DatasmithKeyValueProperty.IsValid())
+				{
+					continue;
+				}
+
+				// Try to find the node that corresponds to the MVR node
+				FString MetaDataMVRUUIDString = FString(DatasmithKeyValueProperty->GetValue());
+				FGuid MetaDataMVRUUID;
+				if (!FGuid::Parse(MetaDataMVRUUIDString, MetaDataMVRUUID))
+				{
+					continue;
+				}
+
+				if (MetaDataMVRUUID == MVRFixture.UUID)
+				{
+					const TSharedPtr<IDatasmithElement> AssociatedElementToRemove = MetaDataElement->GetAssociatedElement();
+
+					if (!AssociatedElementToRemove.IsValid() || !AssociatedElementToRemove->IsA(EDatasmithElementType::Actor))
+					{
+						ensureAlwaysMsgf(0, TEXT("DatasmithMVRNativeTranslator unexpected Element type when trying to remove MVR element in favor of an MVR Scene Actor."), AssociatedElementToRemove->GetName());
+						continue;
+					}
+
+					const TSharedRef<IDatasmithActorElement> ActorElementToRemove = StaticCastSharedRef<IDatasmithActorElement>(AssociatedElementToRemove.ToSharedRef());
+					MetaDataElementsToRemove.Add(MetaDataElement.ToSharedRef());
+					ActorElementsToRemove.Add(ActorElementToRemove);
+				}
+			}
+		}
+	}
+
+	for (const TSharedRef<IDatasmithMetaDataElement>& MetaDataElementToRemove : MetaDataElementsToRemove)
+	{
+		InOutScene->RemoveMetaData(MetaDataElementToRemove);
+	}
+	for (const TSharedRef<IDatasmithActorElement>& ActorElementToRemove : ActorElementsToRemove)
+	{
+		InOutScene->RemoveActor(ActorElementToRemove, EDatasmithActorRemovalRule::RemoveChildren);
+	}
+}
diff --git a/Engine/Plugins/VirtualProduction/DMX/DatasmithMVR/Source/DatasmithMVRTranslator/Private/DatasmithMVRNativeTranslator.h b/Engine/Plugins/VirtualProduction/DMX/DatasmithMVR/Source/DatasmithMVRTranslator/Private/DatasmithMVRNativeTranslator.h
new file mode 100644
--- /dev/null
+++ b/Engine/Plugins/VirtualProduction/DMX/DatasmithMVR/Source/DatasmithMVRTranslator/Private/DatasmithMVRNativeTranslator.h
@@ -0,0 +1,37 @@
+// Copyright Epic Games, Inc. All Rights Reserved.
+
+#pragma once
+
+#include "DatasmithMVRImportOptions.h"
+
+#include "CoreMinimal.h"
+#include "DatasmithNativeTranslator.h"
+
+class UDMXLibrary;
+
+
+/** An override of the datasmith native translator that considers MVR assets aside or in the *_Assets folder of a .udatasmith file */
+class FDatasmithMVRNativeTranslator 
+	: public FDatasmithNativeTranslator
+{
+public:
+	//~ Begin DatasmithNativeTranslator interface
+	virtual FName GetFName() const override { return "DatasmithMVRNativeTranslator"; };
+	virtual bool LoadScene(TSharedRef<IDatasmithScene> InOutScene) override;
+	virtual void GetSceneImportOptions(TArray<TStrongObjectPtr<UDatasmithOptionsBase>>& Options) override;
+	virtual void SetSceneImportOptions(TArray<TStrongObjectPtr<UDatasmithOptionsBase>>& Options) override;
+	//~ End DatasmithNativeTranslator interface
+
+private:
+	/** Tries to find an MVR file next to or inside the datasmith folder */
+	bool FindMVRFile(FString& OutMVRFilePathAndName) const;
+
+	/** Creates a DMX Library from an MVR */
+	UDMXLibrary* CreateDMXLibraryFromMVR(const FString& MVRFilePathAndName) const;
+
+	/** Replaces actors in the Datasmith Scene that correspond to the MVR with the MVR Scene Actor */
+	void ReplaceMVRActorsWithMVRSceneActor(TSharedRef<IDatasmithScene>& InOutScene, UDMXLibrary* DMXLibrary) const;
+
+	/** Datasmith MVR Import Options */
+	TStrongObjectPtr<UDatasmithMVRImportOptions> ImportOptions;
+};
diff --git a/Engine/Plugins/VirtualProduction/DMX/DatasmithMVR/Source/DatasmithMVRTranslator/Private/DatasmithMVRTranslatorModule.cpp b/Engine/Plugins/VirtualProduction/DMX/DatasmithMVR/Source/DatasmithMVRTranslator/Private/DatasmithMVRTranslatorModule.cpp
new file mode 100644
--- /dev/null
+++ b/Engine/Plugins/VirtualProduction/DMX/DatasmithMVR/Source/DatasmithMVRTranslator/Private/DatasmithMVRTranslatorModule.cpp
@@ -0,0 +1,39 @@
+// Copyright Epic Games, Inc. All Rights Reserved.
+
+#include "DatasmithMVRTranslatorModule.h"
+
+#include "DatasmithMVRNativeTranslator.h"
+
+
+#define LOCTEXT_NAMESPACE "DatasmithMVRTranslatorModule"
+
+void FDatasmithMVRTranslatorModule::StartupModule()
+{
+	FModuleManager::Get().LoadModule(TEXT("DatasmithTranslator"));
+
+	// Override the default native translator (the translator for .udatasmith files) and register the translator implemented in this module
+	SetDatasmithMVRNativeTanslatorEnabled(true);
+}
+
+void FDatasmithMVRTranslatorModule::ShutdownModule()
+{
+	Datasmith::UnregisterTranslator<FDatasmithMVRNativeTranslator>();
+}
+
+void FDatasmithMVRTranslatorModule::SetDatasmithMVRNativeTanslatorEnabled(bool bEnabled)
+{
+	if (bEnabled)
+	{
+		Datasmith::UnregisterTranslator<FDatasmithNativeTranslator>();
+		Datasmith::RegisterTranslator<FDatasmithMVRNativeTranslator>();
+	}
+	else
+	{
+		Datasmith::UnregisterTranslator<FDatasmithMVRNativeTranslator>();
+		Datasmith::RegisterTranslator<FDatasmithNativeTranslator>();
+	}
+}
+
+#undef LOCTEXT_NAMESPACE
+
+IMPLEMENT_MODULE(FDatasmithMVRTranslatorModule, DatasmithMVRTranslator)
diff --git a/Engine/Plugins/VirtualProduction/DMX/DatasmithMVR/Source/DatasmithMVRTranslator/Private/DatasmithMVRTranslatorModule.h b/Engine/Plugins/VirtualProduction/DMX/DatasmithMVR/Source/DatasmithMVRTranslator/Private/DatasmithMVRTranslatorModule.h
new file mode 100644
--- /dev/null
+++ b/Engine/Plugins/VirtualProduction/DMX/DatasmithMVR/Source/DatasmithMVRTranslator/Private/DatasmithMVRTranslatorModule.h
@@ -0,0 +1,20 @@
+// Copyright Epic Games, Inc. All Rights Reserved.
+
+#pragma once
+
+#include "IDatasmithMVRTranslatorModule.h"
+
+
+class FDatasmithMVRTranslatorModule
+	: public IDatasmithMVRTranslatorModule
+{
+public:
+	//~ Begin IModuleInterface interface
+	virtual void StartupModule() override;
+	virtual void ShutdownModule() override;
+	//~End IModuleInterface interface
+
+	//~ Begin IDatasmithMVRNativeTanslator inteface
+	virtual void SetDatasmithMVRNativeTanslatorEnabled(bool bEnabled) override;
+	//~ End IDatasmithMVRNativeTranslator interface
+};
diff --git a/Engine/Plugins/VirtualProduction/DMX/DatasmithMVR/Source/DatasmithMVRTranslator/Public/IDatasmithMVRTranslatorModule.h b/Engine/Plugins/VirtualProduction/DMX/DatasmithMVR/Source/DatasmithMVRTranslator/Public/IDatasmithMVRTranslatorModule.h
new file mode 100644
--- /dev/null
+++ b/Engine/Plugins/VirtualProduction/DMX/DatasmithMVR/Source/DatasmithMVRTranslator/Public/IDatasmithMVRTranslatorModule.h
@@ -0,0 +1,22 @@
+// Copyright Epic Games, Inc. All Rights Reserved.
+
+#pragma once
+
+#include "CoreMinimal.h"
+
+#include "Modules/ModuleInterface.h"
+#include "Modules/ModuleManager.h"
+
+
+class DATASMITHMVRTRANSLATOR_API IDatasmithMVRTranslatorModule
+	: public IModuleInterface
+{
+public:
+	static IDatasmithMVRTranslatorModule& Get()
+	{
+		return FModuleManager::Get().GetModuleChecked<IDatasmithMVRTranslatorModule>("DatasmithMVRTranslator");
+	}
+
+	/** Enables or disables the native translator (translator for .udatasmith files) implemented in this module. When the module is started, the translator is enabled */
+	virtual void SetDatasmithMVRNativeTanslatorEnabled(bool bEnabled) = 0;
+};
