DMX - Fix engine crashes when datasmith file is imported via Interchange and the Datasmith MVR plugin is enabled

diff --git a/Engine/Plugins/VirtualProduction/DMX/DatasmithMVR/Source/DatasmithMVRTranslator/DatasmithMVRTranslator.Build.cs b/Engine/Plugins/VirtualProduction/DMX/DatasmithMVR/Source/DatasmithMVRTranslator/DatasmithMVRTranslator.Build.cs
--- a/Engine/Plugins/VirtualProduction/DMX/DatasmithMVR/Source/DatasmithMVRTranslator/DatasmithMVRTranslator.Build.cs
+++ b/Engine/Plugins/VirtualProduction/DMX/DatasmithMVR/Source/DatasmithMVRTranslator/DatasmithMVRTranslator.Build.cs
@@ -26,6 +26,7 @@ public class DatasmithMVRTranslator : ModuleRules
 				"DMXEditor",
 				"DMXRuntime",
 				"Engine",
+				"Slate",
 			}
 		);
 	}
diff --git a/Engine/Plugins/VirtualProduction/DMX/DatasmithMVR/Source/DatasmithMVRTranslator/Private/DatasmithMVRNativeTranslator.cpp b/Engine/Plugins/VirtualProduction/DMX/DatasmithMVR/Source/DatasmithMVRTranslator/Private/DatasmithMVRNativeTranslator.cpp
--- a/Engine/Plugins/VirtualProduction/DMX/DatasmithMVR/Source/DatasmithMVRTranslator/Private/DatasmithMVRNativeTranslator.cpp
+++ b/Engine/Plugins/VirtualProduction/DMX/DatasmithMVR/Source/DatasmithMVRTranslator/Private/DatasmithMVRNativeTranslator.cpp
@@ -9,9 +9,12 @@
 #include "MVR/Types/DMXMVRFixtureNode.h"
 
 #include "AssetRegistry/AssetRegistryModule.h"
+#include "Async/Async.h"
 #include "DatasmithSceneFactory.h"
+#include "Framework/Notifications/NotificationManager.h"
 #include "Misc/Paths.h"
 #include "UObject/Package.h"
+#include "Widgets/Notifications/SNotificationList.h"
 
 
 DECLARE_LOG_CATEGORY_CLASS(LogDatasmithMVRNativeTranslator, Log, All);
@@ -24,7 +27,15 @@ bool FDatasmithMVRNativeTranslator::LoadScene(TSharedRef<IDatasmithScene> InOutS
 		return false;
 	}
 
-	if (ImportOptions.IsValid() && !ImportOptions->bImportMVR)
+	TArray<TObjectPtr<UDatasmithOptionsBase>> ImportOptionsArray;
+	GetSceneImportOptions(ImportOptionsArray);
+
+	const TObjectPtr<UDatasmithOptionsBase>* ImportOptionsPtr = Algo::FindByPredicate(ImportOptionsArray, [](const UDatasmithOptionsBase* ImportOptions)
+		{
+			return ImportOptions->GetClass() == UDatasmithMVRImportOptions::StaticClass();
+		});
+	UDatasmithMVRImportOptions* MVRImportOptions = ImportOptionsPtr ? Cast<UDatasmithMVRImportOptions>(*ImportOptionsPtr) : nullptr;
+	if (MVRImportOptions && !MVRImportOptions->bImportMVR)
 	{
 		// Importing MVR is turned off in options. Acknowledge the scene was successfully imported by the native translator.
 		return true;
@@ -37,18 +48,34 @@ bool FDatasmithMVRNativeTranslator::LoadScene(TSharedRef<IDatasmithScene> InOutS
 		return true;
 	}
 
+	// Interchange imports on its own thread
+	if (!IsInGameThread())
+	{
+		//  MVR support via Interchange is prevented as it ingores any changes made to the scene here.
+		AsyncTask(ENamedThreads::GameThread, []()
+			{
+				const FNotificationInfo Info(NSLOCTEXT("FDatasmithMVRNativeTranslator", "InterchangeNotSupportedNotification", "MVR Import is not supported when importing Datasmith via interchange"));
+				FSlateNotificationManager::Get().AddNotification(Info);
+			});
+
+		// MVR is not supported. Acknowledge the scene was otherwise successfully import via Interchange.
+		return true;
+	}
+
 	InOutScene->SetHost(TEXT("DatasmithMVRNativeTranslator"));
 	InOutScene->SetProductName(TEXT("DatasmithMVRNativeTranslator"));
 
 	UDMXLibrary* DMXLibrary = CreateDMXLibraryFromMVR(MVRFilePathAndName);
-	if (!DMXLibrary)
+	if (DMXLibrary)
 	{
-		const FString MVRFilename = FPaths::GetCleanFilename(MVRFilePathAndName);
-		UE_LOG(LogDatasmithMVRNativeTranslator, Warning, TEXT("Failed to create DMX Library for %s. See previous errors for more info. Creating native datasmith scene instead."), *MVRFilename);
-		return true;
+		ReplaceMVRActorsWithMVRSceneActor(InOutScene, DMXLibrary);
+	}
+	else
+	{
+		const FNotificationInfo Info(NSLOCTEXT("FDatasmithMVRNativeTranslator", "InterchangeNotSupportedNotification", "Failed to import MVR. See log for details."));
+		FSlateNotificationManager::Get().AddNotification(Info);
+		// Acknowledge the scene was successfully imported by the native translator.
 	}
-
-	ReplaceMVRActorsWithMVRSceneActor(InOutScene, DMXLibrary);
 
 	return true;
 }
@@ -57,27 +84,7 @@ void FDatasmithMVRNativeTranslator::GetSceneImportOptions(TArray<TObjectPtr<UDat
 {
 	FDatasmithNativeTranslator::GetSceneImportOptions(Options);
 
-	if (!ImportOptions.IsValid())
-	{
-		const FString& FilePath = GetSource().GetSourceFile();
-
-		ImportOptions = Datasmith::MakeOptions<UDatasmithMVRImportOptions>();
-	}
-
-	Options.Add(ImportOptions.Get());
-}
-
-void FDatasmithMVRNativeTranslator::SetSceneImportOptions(const TArray<TObjectPtr<UDatasmithOptionsBase>>& Options)
-{
-	FDatasmithNativeTranslator::SetSceneImportOptions(Options);
-
-	for (const TObjectPtr<UDatasmithOptionsBase>& OptionPtr : Options)
-	{
-		if (UDatasmithMVRImportOptions* InImportOptions = Cast<UDatasmithMVRImportOptions>(OptionPtr))
-		{
-			ImportOptions.Reset(InImportOptions);
-		}
-	}
+	Options.Add(Datasmith::MakeOptionsPtr<UDatasmithMVRImportOptions>());
 }
 
 bool FDatasmithMVRNativeTranslator::FindMVRFile(FString& OutMVRFilePathAndName) const
@@ -136,7 +143,7 @@ UDMXLibrary* FDatasmithMVRNativeTranslator::CreateDMXLibraryFromMVR(const FStrin
 	return nullptr;
 }
 
-void FDatasmithMVRNativeTranslator::ReplaceMVRActorsWithMVRSceneActor(TSharedRef<IDatasmithScene>& InOutScene, UDMXLibrary* DMXLibrary) const
+void FDatasmithMVRNativeTranslator::ReplaceMVRActorsWithMVRSceneActor(const TSharedRef<IDatasmithScene>& InOutScene, UDMXLibrary* DMXLibrary) const
 {
 	if (!ensureMsgf(DMXLibrary, TEXT("Tried to create an MVR scene in the datasmith scene, but DMXLibrary is null.")))
 	{
diff --git a/Engine/Plugins/VirtualProduction/DMX/DatasmithMVR/Source/DatasmithMVRTranslator/Private/DatasmithMVRNativeTranslator.h b/Engine/Plugins/VirtualProduction/DMX/DatasmithMVR/Source/DatasmithMVRTranslator/Private/DatasmithMVRNativeTranslator.h
--- a/Engine/Plugins/VirtualProduction/DMX/DatasmithMVR/Source/DatasmithMVRTranslator/Private/DatasmithMVRNativeTranslator.h
+++ b/Engine/Plugins/VirtualProduction/DMX/DatasmithMVR/Source/DatasmithMVRTranslator/Private/DatasmithMVRNativeTranslator.h
@@ -19,7 +19,6 @@ public:
 	virtual FName GetFName() const override { return "DatasmithMVRNativeTranslator"; };
 	virtual bool LoadScene(TSharedRef<IDatasmithScene> InOutScene) override;
 	virtual void GetSceneImportOptions(TArray<TObjectPtr<UDatasmithOptionsBase>>& Options) override;
-	virtual void SetSceneImportOptions(const TArray<TObjectPtr<UDatasmithOptionsBase>>& Options) override;
 	//~ End DatasmithNativeTranslator interface
 
 private:
@@ -30,8 +29,5 @@ private:
 	UDMXLibrary* CreateDMXLibraryFromMVR(const FString& MVRFilePathAndName) const;
 
 	/** Replaces actors in the Datasmith Scene that correspond to the MVR with the MVR Scene Actor */
-	void ReplaceMVRActorsWithMVRSceneActor(TSharedRef<IDatasmithScene>& InOutScene, UDMXLibrary* DMXLibrary) const;
-
-	/** Datasmith MVR Import Options */
-	TStrongObjectPtr<UDatasmithMVRImportOptions> ImportOptions;
+	void ReplaceMVRActorsWithMVRSceneActor(const TSharedRef<IDatasmithScene>& InOutScene, UDMXLibrary* DMXLibrary) const;
 };
