DMX - Fix CIS Build Engine Localization

diff --git a/Engine/Plugins/VirtualProduction/DMX/DatasmithMVR/Source/DatasmithMVRTranslator/Private/DatasmithMVRNativeTranslator.cpp b/Engine/Plugins/VirtualProduction/DMX/DatasmithMVR/Source/DatasmithMVRTranslator/Private/DatasmithMVRNativeTranslator.cpp
--- a/Engine/Plugins/VirtualProduction/DMX/DatasmithMVR/Source/DatasmithMVRTranslator/Private/DatasmithMVRNativeTranslator.cpp
+++ b/Engine/Plugins/VirtualProduction/DMX/DatasmithMVR/Source/DatasmithMVRTranslator/Private/DatasmithMVRNativeTranslator.cpp
@@ -72,7 +72,7 @@ bool FDatasmithMVRNativeTranslator::LoadScene(TSharedRef<IDatasmithScene> InOutS
 	}
 	else
 	{
-		const FNotificationInfo Info(NSLOCTEXT("FDatasmithMVRNativeTranslator", "InterchangeNotSupportedNotification", "Failed to import MVR. See log for details."));
+		const FNotificationInfo Info(NSLOCTEXT("FDatasmithMVRNativeTranslator", "FailedToImportMVRFileNotification", "Failed to import MVR. See log for details."));
 		FSlateNotificationManager::Get().AddNotification(Info);
 		// Acknowledge the scene was successfully imported by the native translator.
 	}
