DMX - Fix MVR Import with Datasmith file adds DMX Library in content root

diff --git a/Engine/Plugins/VirtualProduction/DMX/DatasmithMVR/Source/DatasmithMVRTranslator/Private/DatasmithMVRNativeTranslator.cpp b/Engine/Plugins/VirtualProduction/DMX/DatasmithMVR/Source/DatasmithMVRTranslator/Private/DatasmithMVRNativeTranslator.cpp
--- a/Engine/Plugins/VirtualProduction/DMX/DatasmithMVR/Source/DatasmithMVRTranslator/Private/DatasmithMVRNativeTranslator.cpp
+++ b/Engine/Plugins/VirtualProduction/DMX/DatasmithMVR/Source/DatasmithMVRTranslator/Private/DatasmithMVRNativeTranslator.cpp
@@ -132,8 +132,8 @@ UDMXLibrary* FDatasmithMVRNativeTranslator::CreateDMXLibraryFromMVR(const FStrin
 	UDMXLibraryFromMVRFactory* DMXLibraryFromMVRFactory = NewObject<UDMXLibraryFromMVRFactory>();
 
 	const FString SceneName = GetSource().GetSceneName();
-	const FString DMXLibraryName = FPaths::GetBaseFilename(MVRFilePathAndName);
-	const FString PackageName = TEXT("/Game") / SceneName;
+	const FString DMXLibraryName = FPaths::GetBaseFilename(TEXT("DMXLibrary_") + SceneName);
+	const FString PackageName = TEXT("/Game") / SceneName / DMXLibraryName;
 	UPackage* DMXLibraryAssetPackage = CreatePackage(*PackageName);
 	DMXLibraryAssetPackage->FullyLoad();
 
