Fixed some incorrect TSoftObjectPtr types.

diff --git a/Engine/Plugins/Enterprise/DataprepEditor/Source/DataprepEditor/Private/Widgets/SDataprepEditorViewport.cpp b/Engine/Plugins/Enterprise/DataprepEditor/Source/DataprepEditor/Private/Widgets/SDataprepEditorViewport.cpp
--- a/Engine/Plugins/Enterprise/DataprepEditor/Source/DataprepEditor/Private/Widgets/SDataprepEditorViewport.cpp
+++ b/Engine/Plugins/Enterprise/DataprepEditor/Source/DataprepEditor/Private/Widgets/SDataprepEditorViewport.cpp
@@ -1322,10 +1322,10 @@ void SDataprepEditorViewport::LoadDefaultSettings()
 			}
 
 			// Good to go, update the profile's related parameters
-			if( Cast<UTextureCube>( LoadedObject ) != nullptr )
+			if( UTextureCube* LoadedTextureCube = Cast<UTextureCube>( LoadedObject ) )
 			{
 				DataprepViewportSettingProfile.EnvironmentCubeMapPath = EnvironmentCubeMapPath;
-				DataprepViewportSettingProfile.EnvironmentCubeMap = LoadedObject;
+				DataprepViewportSettingProfile.EnvironmentCubeMap = LoadedTextureCube;
 			}
 		}
 	}
diff --git a/Engine/Plugins/Enterprise/DatasmithImporter/Source/DatasmithImporter/Private/DatasmithConsumer.cpp b/Engine/Plugins/Enterprise/DatasmithImporter/Source/DatasmithImporter/Private/DatasmithConsumer.cpp
--- a/Engine/Plugins/Enterprise/DatasmithImporter/Source/DatasmithImporter/Private/DatasmithConsumer.cpp
+++ b/Engine/Plugins/Enterprise/DatasmithImporter/Source/DatasmithImporter/Private/DatasmithConsumer.cpp
@@ -178,7 +178,7 @@ namespace DatasmithConsumerUtils
 
 		for(TPair<FName, TSoftObjectPtr< AssetClass >>& Entry : AssetMap)
 		{
-			if(UObject* Asset = Entry.Value.Get())
+			if(AssetClass* Asset = Entry.Value.Get())
 			{
 				const FString& OutputFolder = GetMarker(Asset, UDataprepContentConsumer::RelativeOutput);
 				if(OutputFolder.Len() > 0)
diff --git a/Engine/Plugins/VirtualProduction/DMX/DatasmithMVR/Source/DatasmithMVRTranslator/Private/DatasmithMVRNativeTranslator.cpp b/Engine/Plugins/VirtualProduction/DMX/DatasmithMVR/Source/DatasmithMVRTranslator/Private/DatasmithMVRNativeTranslator.cpp
--- a/Engine/Plugins/VirtualProduction/DMX/DatasmithMVR/Source/DatasmithMVRTranslator/Private/DatasmithMVRNativeTranslator.cpp
+++ b/Engine/Plugins/VirtualProduction/DMX/DatasmithMVR/Source/DatasmithMVRTranslator/Private/DatasmithMVRNativeTranslator.cpp
@@ -175,7 +175,7 @@ void FDatasmithMVRNativeTranslator::ReplaceMVRActorsWithMVRSceneActor(const TSha
 	CustomMVRSceneActorElement->SetTranslation(FVector::ZeroVector);
 
 	// Set the DMX Library path as metadata of the custom actor element
-	const TSoftObjectPtr<UDMXEntityFixturePatch> SoftDMXLibraryObject = DMXLibrary;
+	const TSoftObjectPtr<UDMXLibrary> SoftDMXLibraryObject = DMXLibrary;
 	const TSharedRef<IDatasmithKeyValueProperty> FixturePatchKeyValueProperty = FDatasmithSceneFactory::CreateKeyValueProperty(TEXT("DMXLibraryPath"));
 	FixturePatchKeyValueProperty->SetValue(*SoftDMXLibraryObject.ToString());
 
